Added Tagged Asset Browser implementation for new Material Wizard. Contained in default enabled plugin.

- Old factories removed from UI if browser is not setup in Material Editor Settings
- New factories with wizard added to UI if browser is setup in Material Editor Settings

diff --git a/Engine/Plugins/Editor/MaterialAssetWizard/MaterialAssetWizard.uplugin b/Engine/Plugins/Editor/MaterialAssetWizard/MaterialAssetWizard.uplugin
new file mode 100644
--- /dev/null
+++ b/Engine/Plugins/Editor/MaterialAssetWizard/MaterialAssetWizard.uplugin
@@ -0,0 +1,24 @@
+﻿{
+	"FileVersion": 3,
+	"Version": 1,
+	"VersionName": "1.0",
+	"FriendlyName": "Material Asset Wizard",
+	"Description": "Enables new factories for Materials and Material Instances. Requires setting Tagged Asset Browser Configuration properties in Material Editor Settings.",
+	"Category": "Editor",
+	"CreatedBy": "Epic Games, Inc.",
+	"CreatedByURL": "https://epicgames.com",
+	"DocsURL": "",
+	"MarketplaceURL": "",
+	"CanContainContent": true,
+	"IsBetaVersion": false,
+	"IsExperimentalVersion": false,
+	"Installed": false,
+	"EnabledByDefault": true,
+	"Modules": [
+		{
+			"Name": "MaterialAssetWizard",
+			"Type": "Editor",
+			"LoadingPhase": "Default"
+		}
+	]
+}
\ No newline at end of file
diff --git a/Engine/Plugins/Editor/MaterialAssetWizard/Source/MaterialAssetWizard/MaterialAssetWizard.Build.cs b/Engine/Plugins/Editor/MaterialAssetWizard/Source/MaterialAssetWizard/MaterialAssetWizard.Build.cs
new file mode 100644
--- /dev/null
+++ b/Engine/Plugins/Editor/MaterialAssetWizard/Source/MaterialAssetWizard/MaterialAssetWizard.Build.cs
@@ -0,0 +1,34 @@
+﻿// Copyright Epic Games, Inc. All Rights Reserved.
+
+using UnrealBuildTool;
+
+public class MaterialAssetWizard : ModuleRules
+{
+	public MaterialAssetWizard(ReadOnlyTargetRules Target) : base(Target)
+	{
+		PCHUsage = ModuleRules.PCHUsageMode.UseExplicitOrSharedPCHs;
+		
+		PublicDependencyModuleNames.AddRange(
+			new string[]
+			{
+				"Core",
+				"UserAssetTagsEditor",
+				"UnrealEd"
+				// ... add other public dependencies that you statically link with here ...
+			}
+			);
+			
+		
+		PrivateDependencyModuleNames.AddRange(
+			new string[]
+			{
+				"CoreUObject",
+				"Engine",
+				"Slate",
+				"SlateCore",
+				"MaterialEditor"
+				// ... add private dependencies that you statically link with here ...	
+			}
+			);
+	}
+}
diff --git a/Engine/Plugins/Editor/MaterialAssetWizard/Source/MaterialAssetWizard/Private/MaterialAssetWizardModule.cpp b/Engine/Plugins/Editor/MaterialAssetWizard/Source/MaterialAssetWizard/Private/MaterialAssetWizardModule.cpp
new file mode 100644
--- /dev/null
+++ b/Engine/Plugins/Editor/MaterialAssetWizard/Source/MaterialAssetWizard/Private/MaterialAssetWizardModule.cpp
@@ -0,0 +1,123 @@
+﻿// Copyright Epic Games, Inc. All Rights Reserved.
+
+#include "MaterialAssetWizardModule.h"
+
+#include "Materials/MaterialInstanceConstant.h"
+#include "Widgets/SUserAssetTagBrowserSelectedAssetDetails.h"
+#include "Widgets/Text/STextBlock.h"
+#include "Materials/Material.h"
+
+#define LOCTEXT_NAMESPACE "MaterialAssetWizardModule"
+
+void FMaterialAssetWizardModule::StartupModule()
+{
+	// Materials
+	{
+		FAssetDetailsDisplayInfo MaterialInfo;
+
+		// Example
+#if 0
+		// Number of emitters
+		FDisplayedPropertyData NumEmittersProperty;
+		NumEmittersProperty.ShouldDisplayPropertyDelegate = FDisplayedPropertyData::FShouldDisplayProperty::CreateLambda([](const FAssetData& AssetData) -> bool
+		{
+			return AssetData.FindTag("NumEmitters");
+		});
+		NumEmittersProperty.NameWidgetDelegate = FDisplayedPropertyData::FGenerateWidget::CreateLambda([](const FAssetData& AssetData) -> TSharedRef<SWidget>
+		{
+			return SNew(STextBlock)
+			.Text(LOCTEXT("NumberOfEmitters", "Number of Emitters"));
+		});
+		NumEmittersProperty.ValueWidgetDelegate =  FDisplayedPropertyData::FGenerateWidget::CreateLambda([](const FAssetData& AssetData) -> TSharedRef<SWidget>
+		{
+			int32 NumEmitters = INDEX_NONE;
+			if(AssetData.GetTagValue("NumEmitters", NumEmitters))
+			{
+				return SNew(STextBlock).Text(FText::AsNumber(NumEmitters));
+			}
+	
+			return SNullWidget::NullWidget;
+		});
+		
+		MaterialInfo.DisplayedProperties =
+		{
+			NumEmittersProperty
+		};
+		
+		MaterialInfo.GetDescriptionDelegate = FAssetDetailsDisplayInfo::FGetDescription::CreateLambda([](const FAssetData& AssetData)
+		{
+			FName DescriptionTagName("TemplateAssetDescription");
+			if(AssetData.FindTag(DescriptionTagName))
+			{
+				FString Description;
+				AssetData.GetTagValue(DescriptionTagName, Description);
+				return FText::FromString(Description);
+			}
+		
+			return FText::GetEmpty();
+		});
+#endif
+		
+		FTaggedAssetBrowserDetailsDisplayDatabase::RegisterClass(UMaterial::StaticClass(), MaterialInfo);
+	}
+
+	// Material Instances
+	{
+		FAssetDetailsDisplayInfo MaterialInstanceInfo;
+
+		// Example
+#if 0
+		// Number of emitters
+		FDisplayedPropertyData NumEmittersProperty;
+		NumEmittersProperty.ShouldDisplayPropertyDelegate = FDisplayedPropertyData::FShouldDisplayProperty::CreateLambda([](const FAssetData& AssetData) -> bool
+		{
+			return AssetData.FindTag("NumEmitters");
+		});
+		NumEmittersProperty.NameWidgetDelegate = FDisplayedPropertyData::FGenerateWidget::CreateLambda([](const FAssetData& AssetData) -> TSharedRef<SWidget>
+		{
+			return SNew(STextBlock)
+			.Text(LOCTEXT("NumberOfEmitters", "Number of Emitters"));
+		});
+		NumEmittersProperty.ValueWidgetDelegate =  FDisplayedPropertyData::FGenerateWidget::CreateLambda([](const FAssetData& AssetData) -> TSharedRef<SWidget>
+		{
+			int32 NumEmitters = INDEX_NONE;
+			if(AssetData.GetTagValue("NumEmitters", NumEmitters))
+			{
+				return SNew(STextBlock).Text(FText::AsNumber(NumEmitters));
+			}
+	
+			return SNullWidget::NullWidget;
+		});
+		
+		MaterialInstanceInfo.DisplayedProperties =
+		{
+			NumEmittersProperty
+		};
+		
+		MaterialInstanceInfo.GetDescriptionDelegate = FAssetDetailsDisplayInfo::FGetDescription::CreateLambda([](const FAssetData& AssetData)
+		{
+			FName DescriptionTagName("TemplateAssetDescription");
+			if(AssetData.FindTag(DescriptionTagName))
+			{
+				FString Description;
+				AssetData.GetTagValue(DescriptionTagName, Description);
+				return FText::FromString(Description);
+			}
+		
+			return FText::GetEmpty();
+		});
+#endif
+		
+		FTaggedAssetBrowserDetailsDisplayDatabase::RegisterClass(UMaterialInstanceConstant::StaticClass(), MaterialInstanceInfo);
+	}
+}
+
+void FMaterialAssetWizardModule::ShutdownModule()
+{
+	// This function may be called during shutdown to clean up your module.  For modules that support dynamic reloading,
+	// we call this function before unloading the module.
+}
+
+#undef LOCTEXT_NAMESPACE
+	
+IMPLEMENT_MODULE(FMaterialAssetWizardModule, MaterialAssetWizard)
\ No newline at end of file
diff --git a/Engine/Plugins/Editor/MaterialAssetWizard/Source/MaterialAssetWizard/Private/MaterialFactoryWithAssetWizard.cpp b/Engine/Plugins/Editor/MaterialAssetWizard/Source/MaterialAssetWizard/Private/MaterialFactoryWithAssetWizard.cpp
new file mode 100644
--- /dev/null
+++ b/Engine/Plugins/Editor/MaterialAssetWizard/Source/MaterialAssetWizard/Private/MaterialFactoryWithAssetWizard.cpp
@@ -0,0 +1,167 @@
+﻿// Copyright Epic Games, Inc. All Rights Reserved.
+
+#include "MaterialFactoryWithAssetWizard.h"
+
+#include "EditorDirectories.h"
+#include "MaterialEditingLibrary.h"
+#include "MaterialEditorSettings.h"
+#include "Assets/TaggedAssetBrowserConfiguration.h"
+#include "Interfaces/IMainFrameModule.h"
+#include "Widgets/STaggedAssetBrowser.h"
+#include "Widgets/STaggedAssetBrowserCreateAssetWindow.h"
+#include "Materials/Material.h"
+#include "Materials/MaterialInstanceConstant.h"
+
+#define LOCTEXT_NAMESPACE "MaterialFactoryWithAssetWizard"
+
+bool UMaterialFactoryWithAssetWizardBase::ConfigureProperties()
+{
+	if (SupportedClass == nullptr)
+	{
+		return false;
+	}
+	
+	TArray<FAssetData> SelectedAssetData;
+	
+	FSoftObjectPath Path = GetDefault<UMaterialEditorSettings>()->MaterialWizardConfiguration;
+	UObject* ConfigUObject = Path.TryLoad();
+	if(UTaggedAssetBrowserConfiguration* ConfigurationAsset = Cast<UTaggedAssetBrowserConfiguration>(ConfigUObject))
+	{ 
+		IMainFrameModule& MainFrame = FModuleManager::LoadModuleChecked<IMainFrameModule>("MainFrame");
+		TSharedPtr<SWindow>	ParentWindow = MainFrame.GetParentWindow();
+
+		SWindow::FArguments WindowArgs = CreateWindowArguments();
+		STaggedAssetBrowser::FArguments AssetBrowserArgs = CreateTaggedAssetBrowserArguments();
+		
+		STaggedAssetBrowserWindow::FArguments AssetBrowserWindowArgs;
+		AssetBrowserWindowArgs.AssetBrowserArgs(AssetBrowserArgs);
+		AssetBrowserWindowArgs.WindowArgs(WindowArgs);
+		
+		TSharedRef<STaggedAssetBrowserCreateAssetWindow> CreateAssetBrowserWindow = SNew(STaggedAssetBrowserCreateAssetWindow, *ConfigurationAsset, *SupportedClass.Get())
+			.AssetBrowserWindowArgs(AssetBrowserWindowArgs);
+		
+		FSlateApplication::Get().AddModalWindow(CreateAssetBrowserWindow, ParentWindow);
+
+		if(CreateAssetBrowserWindow->ShouldProceedWithAction() == false)
+		{
+			return false;	
+		}
+
+		SelectedAssetData = CreateAssetBrowserWindow->GetSelectedAssets();
+	}
+
+	if(SelectedAssetData.Num() == 1)
+	{
+		FAssetData SelectedAsset = SelectedAssetData[0];
+		ensure(SelectedAsset.GetClass() == UMaterial::StaticClass());
+		
+		UMaterial* MaterialAsset = Cast<UMaterial>(SelectedAsset.GetAsset());
+			
+		BaseMaterial = MaterialAsset;
+	}
+	
+	return true;
+}
+
+SWindow::FArguments UMaterialFactoryWithAssetWizardBase::CreateWindowArguments() const
+{
+	SWindow::FArguments Result;
+	Result.SupportsMaximize(false);
+	Result.SupportsMinimize(false);
+	Result.ClientSize(FVector2D(1400, 750));
+	Result.SizingRule(ESizingRule::UserSized);
+	Result.Title(WindowTitle);
+		
+	return Result;
+}
+
+STaggedAssetBrowser::FArguments UMaterialFactoryWithAssetWizardBase::CreateTaggedAssetBrowserArguments() const
+{
+	FDefaultDetailsTabConfiguration DefaultDetailsTabConfiguration;
+	DefaultDetailsTabConfiguration.bUseDefaultDetailsTab = true;
+	DefaultDetailsTabConfiguration.EmptySelectionMessage = EmptySelectionMessage;
+		
+	STaggedAssetBrowser::FInterfaceOverrideProfiles OverrideProfiles;
+	OverrideProfiles.AssetViewOptionsProfileName = FName("TaggedAssetBrowser");
+	OverrideProfiles.DefaultFilterMenuExpansion = EAssetTypeCategories::Materials;
+	OverrideProfiles.FilterBarSaveName = FName("TaggedAssetBrowser.MaterialFactory");
+		
+	STaggedAssetBrowser::FArguments Result;
+	Result.AvailableClasses({ UMaterial::StaticClass() });
+	Result.DefaultDetailsTabConfiguration(DefaultDetailsTabConfiguration);
+	Result.InterfaceOverrideProfiles(OverrideProfiles);
+	// @TODO (ME): This is not ideal. Ideally the factory would have access to creation path and/or property handle (if created from property asset picker -> create menu)
+	// then we wouldn't need to use GetLastDirectory which is relatively unsafe
+	Result.AdditionalReferencingAssets({ FAssetData("TmpAsset", FName(FEditorDirectories::Get().GetLastDirectory(ELastDirectory::NEW_ASSET)), "TmpAsset", UObject::StaticClass()->GetClassPathName())});
+
+	return Result;
+}
+
+UMaterialFactoryWithAssetWizard::UMaterialFactoryWithAssetWizard()
+{
+	SupportedClass = UMaterial::StaticClass();
+	bCreateNew = true;
+	bEditAfterNew = true;
+
+	WindowTitle = LOCTEXT("MaterialAssetBrowserWindowTitle", "Create Material - Select another material as a base");
+	EmptySelectionMessage = LOCTEXT("EmptyMaterialFactorySelectionMessage", "Select a material as a starting point for your new material.\n");
+}
+
+UObject* UMaterialFactoryWithAssetWizard::FactoryCreateNew(UClass* Class, UObject* InParent, FName Name, EObjectFlags Flags, UObject* Context, FFeedbackContext* Warn)
+{
+	check(Class->IsChildOf(UMaterial::StaticClass()));
+
+	const UMaterialEditorSettings* Settings = GetDefault<UMaterialEditorSettings>();
+	check(Settings);
+
+	UMaterial* NewMaterial;
+
+	// Duplicate the selected material
+	if (BaseMaterial.IsValid())
+	{
+		NewMaterial = Cast<UMaterial>(StaticDuplicateObject(BaseMaterial.Get(), InParent, Name, Flags, Class));
+	}
+	// Or create an empty material as a fallback
+	else
+	{
+		NewMaterial = NewObject<UMaterial>(InParent, Class, Name, Flags | RF_Transactional);
+	}
+
+	return NewMaterial;
+}
+
+bool UMaterialFactoryWithAssetWizard::ShouldShowInNewMenu() const
+{
+	return UFactory::ShouldShowInNewMenu() && GetDefault<UMaterialEditorSettings>()->MaterialWizardConfiguration != nullptr;
+}
+
+UMaterialInstanceConstantFactoryWithAssetWizard::UMaterialInstanceConstantFactoryWithAssetWizard()
+{
+	SupportedClass = UMaterialInstanceConstant::StaticClass();
+	bCreateNew = true;
+	bEditAfterNew = true;
+
+	WindowTitle = LOCTEXT("MaterialInstanceAssetBrowserWindowTitle", "Create Material Instance - Select another Material as a base");
+	EmptySelectionMessage = LOCTEXT("EmptyMaterialInstanceFactorySelectionMessage", "Select a material as a starting point for your new material instance.\n");
+}
+
+UObject* UMaterialInstanceConstantFactoryWithAssetWizard::FactoryCreateNew(UClass* Class, UObject* InParent, FName Name, EObjectFlags Flags,	UObject* Context, FFeedbackContext* Warn)
+{
+	check(Class->IsChildOf(UMaterialInstanceConstant::StaticClass()));
+
+	UMaterialInstanceConstant* NewMaterialInstance = NewObject<UMaterialInstanceConstant>(InParent, Class, Name, Flags | RF_Transactional);
+	
+	if (BaseMaterial.IsValid())
+	{
+		UMaterialEditingLibrary::SetMaterialInstanceParent(NewMaterialInstance, BaseMaterial.Get());
+	}
+
+	return NewMaterialInstance;
+}
+
+bool UMaterialInstanceConstantFactoryWithAssetWizard::ShouldShowInNewMenu() const
+{
+	return UFactory::ShouldShowInNewMenu() && GetDefault<UMaterialEditorSettings>()->MaterialInstanceWizardConfiguration != nullptr;
+}
+
+#undef LOCTEXT_NAMESPACE
diff --git a/Engine/Plugins/Editor/MaterialAssetWizard/Source/MaterialAssetWizard/Public/MaterialAssetWizardModule.h b/Engine/Plugins/Editor/MaterialAssetWizard/Source/MaterialAssetWizard/Public/MaterialAssetWizardModule.h
new file mode 100644
--- /dev/null
+++ b/Engine/Plugins/Editor/MaterialAssetWizard/Source/MaterialAssetWizard/Public/MaterialAssetWizardModule.h
@@ -0,0 +1,14 @@
+﻿// Copyright Epic Games, Inc. All Rights Reserved.
+
+#pragma once
+
+#include "Modules/ModuleManager.h"
+
+class FMaterialAssetWizardModule : public IModuleInterface
+{
+public:
+
+	/** IModuleInterface implementation */
+	virtual void StartupModule() override;
+	virtual void ShutdownModule() override;
+};
diff --git a/Engine/Plugins/Editor/MaterialAssetWizard/Source/MaterialAssetWizard/Public/MaterialFactoryWithAssetWizard.h b/Engine/Plugins/Editor/MaterialAssetWizard/Source/MaterialAssetWizard/Public/MaterialFactoryWithAssetWizard.h
new file mode 100644
--- /dev/null
+++ b/Engine/Plugins/Editor/MaterialAssetWizard/Source/MaterialAssetWizard/Public/MaterialFactoryWithAssetWizard.h
@@ -0,0 +1,63 @@
+﻿// Copyright Epic Games, Inc. All Rights Reserved.
+
+#pragma once
+
+#include "CoreMinimal.h"
+#include "Factories/Factory.h"
+#include "Widgets/STaggedAssetBrowser.h"
+#include "MaterialFactoryWithAssetWizard.generated.h"
+
+class STaggedAssetBrowserCreateAssetWindow;
+class UMaterial;
+
+/** Abstract base factory that summons material wizard. */
+UCLASS(Abstract)
+class MATERIALASSETWIZARD_API UMaterialFactoryWithAssetWizardBase : public UFactory
+{
+	GENERATED_BODY()
+
+public:
+	UMaterialFactoryWithAssetWizardBase() = default;
+	
+	virtual bool ConfigureProperties() override;
+	virtual bool ShouldShowInNewMenu() const override { return false; }
+
+private:
+	SWindow::FArguments CreateWindowArguments() const;
+	STaggedAssetBrowser::FArguments CreateTaggedAssetBrowserArguments() const;
+	
+protected:
+	UPROPERTY()
+	FText WindowTitle;
+
+	UPROPERTY()
+	FText EmptySelectionMessage;
+	
+	TWeakObjectPtr<UMaterial> BaseMaterial;
+};
+
+/** Concrete material factory that will create a new material based on the selected material, or an empty one. */
+UCLASS()
+class MATERIALASSETWIZARD_API UMaterialFactoryWithAssetWizard : public UMaterialFactoryWithAssetWizardBase
+{
+	GENERATED_BODY()
+
+public:
+	UMaterialFactoryWithAssetWizard();
+	
+	virtual UObject* FactoryCreateNew(UClass* Class, UObject* InParent, FName Name, EObjectFlags Flags, UObject* Context, FFeedbackContext* Warn) override;
+	virtual bool ShouldShowInNewMenu() const override;
+};
+
+/** Concrete material instance constant factory that will create a new material instance based on the selected material, or an empty one. */
+UCLASS()
+class MATERIALASSETWIZARD_API UMaterialInstanceConstantFactoryWithAssetWizard : public UMaterialFactoryWithAssetWizardBase
+{
+	GENERATED_BODY()
+
+public:
+	UMaterialInstanceConstantFactoryWithAssetWizard();
+	
+	virtual UObject* FactoryCreateNew(UClass* Class, UObject* InParent, FName Name, EObjectFlags Flags, UObject* Context, FFeedbackContext* Warn) override;
+	virtual bool ShouldShowInNewMenu() const override;
+};
diff --git a/Engine/Source/Editor/MaterialEditor/Public/MaterialEditorSettings.h b/Engine/Source/Editor/MaterialEditor/Public/MaterialEditorSettings.h
--- a/Engine/Source/Editor/MaterialEditor/Public/MaterialEditorSettings.h
+++ b/Engine/Source/Editor/MaterialEditor/Public/MaterialEditorSettings.h
@@ -105,6 +105,14 @@ class UMaterialEditorSettings : public UObject
 	GENERATED_BODY()
 
 public:
+
+	/** The configuration used in the Material Asset Wizard when creating a new Material. If not set, Content Browser uses the old factory without wizard. */
+	UPROPERTY(config, EditAnywhere, Category = Wizard, meta=(AllowedClasses="/Script/UserAssetTagsEditor.TaggedAssetBrowserConfiguration"))
+	FSoftObjectPath MaterialWizardConfiguration;
+
+	/** The configuration used in the Material Asset Wizard when creating a new Material Instance. If not set, Content Browser uses the old factory without wizard. */
+	UPROPERTY(config, EditAnywhere, Category = Wizard, meta=(AllowedClasses="/Script/UserAssetTagsEditor.TaggedAssetBrowserConfiguration"))
+	FSoftObjectPath MaterialInstanceWizardConfiguration;
 	
 	/**
 	Allow ignoring compilation errors of platform shaders and derived materials.
diff --git a/Engine/Source/Editor/UnrealEd/Classes/Factories/MaterialFactoryNew.h b/Engine/Source/Editor/UnrealEd/Classes/Factories/MaterialFactoryNew.h
--- a/Engine/Source/Editor/UnrealEd/Classes/Factories/MaterialFactoryNew.h
+++ b/Engine/Source/Editor/UnrealEd/Classes/Factories/MaterialFactoryNew.h
@@ -22,5 +22,6 @@ class UMaterialFactoryNew : public UFactory
 
 	//~ Begin UFactory Interface
 	virtual UObject* FactoryCreateNew(UClass* Class,UObject* InParent,FName Name,EObjectFlags Flags,UObject* Context,FFeedbackContext* Warn) override;
+	virtual bool ShouldShowInNewMenu() const override;
 	//~ Begin UFactory Interface	
 };
diff --git a/Engine/Source/Editor/UnrealEd/Classes/Factories/MaterialInstanceConstantFactoryNew.h b/Engine/Source/Editor/UnrealEd/Classes/Factories/MaterialInstanceConstantFactoryNew.h
--- a/Engine/Source/Editor/UnrealEd/Classes/Factories/MaterialInstanceConstantFactoryNew.h
+++ b/Engine/Source/Editor/UnrealEd/Classes/Factories/MaterialInstanceConstantFactoryNew.h
@@ -21,6 +21,7 @@ class UMaterialInstanceConstantFactoryNew : public UFactory
 
 	//~ Begin UFactory Interface
 	UNREALED_API virtual UObject* FactoryCreateNew(UClass* Class,UObject* InParent,FName Name,EObjectFlags Flags,UObject* Context,FFeedbackContext* Warn) override;
+	virtual bool ShouldShowInNewMenu() const override;
 	//~ Begin UFactory Interface	
 };
 
diff --git a/Engine/Source/Editor/UnrealEd/Private/Factories/EditorFactories.cpp b/Engine/Source/Editor/UnrealEd/Private/Factories/EditorFactories.cpp
--- a/Engine/Source/Editor/UnrealEd/Private/Factories/EditorFactories.cpp
+++ b/Engine/Source/Editor/UnrealEd/Private/Factories/EditorFactories.cpp
@@ -283,6 +283,7 @@
 #include "TextureReferenceResolver.h"
 #include "UDIMUtilities.h"
 #include "FileHelpers.h"
+#include "MaterialEditorSettings.h"
 #include "StaticMeshOperations.h"
 #include "UObject/OverridableManager.h"
 
@@ -489,6 +490,11 @@ UObject* UMaterialInstanceConstantFactoryNew::FactoryCreateNew(UClass* Class,UOb
 	return MIC;
 }
 
+bool UMaterialInstanceConstantFactoryNew::ShouldShowInNewMenu() const
+{
+	return Super::ShouldShowInNewMenu() && GetDefault<UMaterialEditorSettings>()->MaterialInstanceWizardConfiguration == nullptr;
+}
+
 /*------------------------------------------------------------------------------
 	UMaterialFactoryNew implementation.
 ------------------------------------------------------------------------------*/
@@ -540,6 +546,11 @@ UObject* UMaterialFactoryNew::FactoryCreateNew(UClass* Class,UObject* InParent,F
 	return NewMaterial;
 }
 
+bool UMaterialFactoryNew::ShouldShowInNewMenu() const
+{
+	return Super::ShouldShowInNewMenu() && GetDefault<UMaterialEditorSettings>()->MaterialWizardConfiguration == nullptr;
+}
+
 
 /*------------------------------------------------------------------------------
 	UMaterialFunctionFactoryNew implementation.
