Fixed Material Async Wizard factories

diff --git a/Engine/Plugins/Editor/MaterialAssetWizard/Source/MaterialAssetWizard/Private/MaterialFactoryWithAssetWizard.cpp b/Engine/Plugins/Editor/MaterialAssetWizard/Source/MaterialAssetWizard/Private/MaterialFactoryWithAssetWizard.cpp
--- a/Engine/Plugins/Editor/MaterialAssetWizard/Source/MaterialAssetWizard/Private/MaterialFactoryWithAssetWizard.cpp
+++ b/Engine/Plugins/Editor/MaterialAssetWizard/Source/MaterialAssetWizard/Private/MaterialFactoryWithAssetWizard.cpp
@@ -27,8 +27,7 @@ bool UMaterialFactoryWithAssetWizardBase::ConfigureProperties()
 
 bool UMaterialFactoryWithAssetWizardBase::ConfigurePropertiesAsync(FOnFactoryConfigurePropertiesAsyncComplete OnComplete, FOnFactoryConfigurePropertiesAsyncCancelled OnCancelled)
 {	
-	FSoftObjectPath Path = GetDefault<UMaterialEditorSettings>()->MaterialWizardConfiguration;
-	UObject* ConfigUObject = Path.TryLoad();
+	UObject* ConfigUObject = ConfigurationPath.TryLoad();
 	if(UTaggedAssetBrowserConfiguration* ConfigurationAsset = Cast<UTaggedAssetBrowserConfiguration>(ConfigUObject))
 	{
 		SWindow::FArguments WindowArgs = CreateWindowArguments();
@@ -92,14 +91,17 @@ STaggedAssetBrowser::FArguments UMaterialFactoryWithAssetWizardBase::CreateTagge
 	OverrideProfiles.AssetViewOptionsProfileName = FName("TaggedAssetBrowser");
 	OverrideProfiles.DefaultFilterMenuExpansion = EAssetTypeCategories::Materials;
 	OverrideProfiles.FilterBarSaveName = FName("TaggedAssetBrowser.MaterialFactory");
-		
+
+	FString PackagePath = FEditorDirectories::Get().GetLastDirectory(ELastDirectory::NEW_ASSET);
+	FString AssetName = "TmpAsset";
+	FName PackageName = FName(PackagePath + "/" + AssetName);
+	FAssetData PseudoReferencingAsset(PackageName, FName(PackagePath), FName(AssetName), UObject::StaticClass()->GetClassPathName());
+	
 	STaggedAssetBrowser::FArguments Result;
 	Result.AvailableClasses({ UMaterial::StaticClass() });
 	Result.DefaultDetailsTabConfiguration(DefaultDetailsTabConfiguration);
 	Result.InterfaceOverrideProfiles(OverrideProfiles);
-	// @TODO (ME): This is not ideal. Ideally the factory would have access to creation path and/or property handle (if created from property asset picker -> create menu)
-	// then we wouldn't need to use GetLastDirectory which is relatively unsafe
-	Result.AdditionalReferencingAssets({ FAssetData("TmpAsset", FName(FEditorDirectories::Get().GetLastDirectory(ELastDirectory::NEW_ASSET)), "TmpAsset", UObject::StaticClass()->GetClassPathName())});
+	Result.AdditionalReferencingAssets({ PseudoReferencingAsset });
 
 	return Result;
 }
@@ -109,7 +111,9 @@ UMaterialFactoryWithAssetWizard::UMaterialFactoryWithAssetWizard()
 	SupportedClass = UMaterial::StaticClass();
 	bCreateNew = true;
 	bEditAfterNew = true;
+	SupportedWorkflows = (uint8) (EFactoryCreateWorkflow::Default | EFactoryCreateWorkflow::Asynchronous);
 
+	ConfigurationPath = GetDefault<UMaterialEditorSettings>()->MaterialWizardConfiguration;
 	WindowTitle = LOCTEXT("MaterialAssetBrowserWindowTitle", "Create Material - Select another material as a base");
 	EmptySelectionMessage = LOCTEXT("EmptyMaterialFactorySelectionMessage", "Select a material as a starting point for your new material.\n");
 }
@@ -147,7 +151,9 @@ UMaterialInstanceConstantFactoryWithAssetWizard::UMaterialInstanceConstantFactor
 	SupportedClass = UMaterialInstanceConstant::StaticClass();
 	bCreateNew = true;
 	bEditAfterNew = true;
+	SupportedWorkflows = (uint8) (EFactoryCreateWorkflow::Default | EFactoryCreateWorkflow::Asynchronous);
 
+	ConfigurationPath = GetDefault<UMaterialEditorSettings>()->MaterialInstanceWizardConfiguration;
 	WindowTitle = LOCTEXT("MaterialInstanceAssetBrowserWindowTitle", "Create Material Instance - Select another Material as a base");
 	EmptySelectionMessage = LOCTEXT("EmptyMaterialInstanceFactorySelectionMessage", "Select a material as a starting point for your new material instance.\n");
 }
diff --git a/Engine/Plugins/Editor/MaterialAssetWizard/Source/MaterialAssetWizard/Public/MaterialFactoryWithAssetWizard.h b/Engine/Plugins/Editor/MaterialAssetWizard/Source/MaterialAssetWizard/Public/MaterialFactoryWithAssetWizard.h
--- a/Engine/Plugins/Editor/MaterialAssetWizard/Source/MaterialAssetWizard/Public/MaterialFactoryWithAssetWizard.h
+++ b/Engine/Plugins/Editor/MaterialAssetWizard/Source/MaterialAssetWizard/Public/MaterialFactoryWithAssetWizard.h
@@ -28,14 +28,19 @@ private:
 	
 	SWindow::FArguments CreateWindowArguments() const;
 	STaggedAssetBrowser::FArguments CreateTaggedAssetBrowserArguments() const;
-	
+
 protected:
+	/** The path to the tagged asset browser configuration asset. Needs to be valid. */
+	UPROPERTY()
+	FSoftObjectPath ConfigurationPath;
+	
 	UPROPERTY()
 	FText WindowTitle;
 
 	UPROPERTY()
 	FText EmptySelectionMessage;
-	
+
+	/** The selected base material, if any. */
 	TWeakObjectPtr<UMaterial> BaseMaterial;
 };
 
