Ensure material resources and shaders are compiled when creating new material assets so their thumbnail renders correctly.

- This fixes a bug where all newly created materials would render their thumbnail with the default material.
- Without these changes there was no rendering resource and UMaterialInstanceThumbnailRenderer::CanVisualizeAsset would always return true since IsCompiling always returned false.
- Submit background shader compilation jobs, and let the thumbnail system handle calling CanVisualizeAsset to see if we are ready to render.

[RN] minor | Rendering - Materials and Shaders

#Synced-CL 49034721

diff --git a/Engine/Plugins/Editor/MaterialAssetWizard/Source/MaterialAssetWizard/Private/MaterialFactoryWithAssetWizard.cpp b/Engine/Plugins/Editor/MaterialAssetWizard/Source/MaterialAssetWizard/Private/MaterialFactoryWithAssetWizard.cpp
--- a/Engine/Plugins/Editor/MaterialAssetWizard/Source/MaterialAssetWizard/Private/MaterialFactoryWithAssetWizard.cpp
+++ b/Engine/Plugins/Editor/MaterialAssetWizard/Source/MaterialAssetWizard/Private/MaterialFactoryWithAssetWizard.cpp
@@ -138,6 +138,8 @@ UObject* UMaterialFactoryWithAssetWizard::FactoryCreateNew(UClass* Class, UObjec
 		NewMaterial = NewObject<UMaterial>(InParent, Class, Name, Flags | RF_Transactional);
 	}
 
+	NewMaterial->CacheShaders(EMaterialShaderPrecompileMode::Background);
+
 	return NewMaterial;
 }
 
diff --git a/Engine/Source/Editor/UnrealEd/Private/Factories/EditorFactories.cpp b/Engine/Source/Editor/UnrealEd/Private/Factories/EditorFactories.cpp
--- a/Engine/Source/Editor/UnrealEd/Private/Factories/EditorFactories.cpp
+++ b/Engine/Source/Editor/UnrealEd/Private/Factories/EditorFactories.cpp
@@ -586,7 +586,7 @@ UObject* UUIMaterialFactoryNew::FactoryCreateNew(UClass* Class, UObject* InParen
 	NewMaterial->bAutomaticallySetUsageInEditor = false;
 	NewMaterial->MaterialDomain = EMaterialDomain::MD_UI;
 	NewMaterial->SetUsageByFlag(MATUSAGE_StaticMesh, false);
-	//NewMaterial->CacheShaders(EMaterialShaderPrecompileMode::Synchronous);
+	NewMaterial->CacheShaders(EMaterialShaderPrecompileMode::Background);
 
 	return NewMaterial;
 }
