Reworded plugin description

#rnx

#Synced-CL 50489927

diff --git a/Engine/Plugins/Editor/MaterialAssetWizard/MaterialAssetWizard.uplugin b/Engine/Plugins/Editor/MaterialAssetWizard/MaterialAssetWizard.uplugin
--- a/Engine/Plugins/Editor/MaterialAssetWizard/MaterialAssetWizard.uplugin
+++ b/Engine/Plugins/Editor/MaterialAssetWizard/MaterialAssetWizard.uplugin
@@ -3,7 +3,7 @@
 	"Version": 1,
 	"VersionName": "1.0",
 	"FriendlyName": "Material Asset Wizard",
-	"Description": "Enables new factories for Materials and Material Instances. Requires setting Tagged Asset Browser Configuration properties in Material Editor Settings.",
+	"Description": "Enables new factories for Materials and Material Instances.\nRequires setting Tagged Asset Browser Configuration properties under Editor Preferences -> Material Editor Settings.",
 	"Category": "Editor",
 	"CreatedBy": "Epic Games, Inc.",
 	"CreatedByURL": "https://epicgames.com",
